Resurrected the Streaming Pause Rendering code and fixed up to display a simple (but modifiable) loading screen during blocking loads.

Moved into own module to bypass dependancy issue with MoviePlayer and Engine.
All work now done on game thread and just kicks off the loading screen/movie player code to handle rendering.
Some new delegates in UEngine allow licencees to easily override this default behavior.

diff --git a/Engine/Source/Runtime/Core/Private/HAL/ConsoleManager.cpp b/Engine/Source/Runtime/Core/Private/HAL/ConsoleManager.cpp
--- a/Engine/Source/Runtime/Core/Private/HAL/ConsoleManager.cpp
+++ b/Engine/Source/Runtime/Core/Private/HAL/ConsoleManager.cpp
@@ -2213,3 +2213,9 @@ static TAutoConsoleVariable<float> CVarDemosaicVposOffset(
 	0.0f,
 	TEXT("This offset is added to the rasterized position used for demosaic in the ES2 tonemapping shader. It exists to workaround driver bugs on some Android devices that have a half-pixel offset."),
 	ECVF_RenderThreadSafe);
+
+static TAutoConsoleVariable<int32> CVarRenderLastFrameInStreamingPause(
+	TEXT("r.RenderLastFrameInStreamingPause"),
+	1,
+	TEXT("If 1 the previous frame is displayed during streaming pause. If zero the screen is left black."),
+	ECVF_RenderThreadSafe);
\ No newline at end of file
diff --git a/Engine/Source/Runtime/Engine/Classes/Engine/Engine.h b/Engine/Source/Runtime/Engine/Classes/Engine/Engine.h
--- a/Engine/Source/Runtime/Engine/Classes/Engine/Engine.h
+++ b/Engine/Source/Runtime/Engine/Classes/Engine/Engine.h
@@ -462,6 +462,9 @@ namespace EMatineeCaptureType
 
 class IAnalyticsProvider;
 
+DECLARE_DELEGATE_OneParam(FBeginStreamingPauseDelegate, FViewport*);
+DECLARE_DELEGATE(FEndStreamingPauseDelegate);
+
 UCLASS(abstract, config=Engine, transient)
 class ENGINE_API UEngine : public UObject, public FExec
 {
@@ -1341,6 +1344,14 @@ public:
 	UPROPERTY(transient)
 	float SelectionHighlightIntensityBillboards;
 
+	/** Delegate handling when streaming pause begins. Set initially in FStreamingPauseRenderingModule::StartupModule() but can then be overriden by games. */
+	void RegisterBeginStreamingPauseRenderingDelegate( FBeginStreamingPauseDelegate* InDelegate );
+	FBeginStreamingPauseDelegate* BeginStreamingPauseDelegate;
+
+	/** Delegate handling when streaming pause ends. Set initially in FStreamingPauseRenderingModule::StartupModule() but can then be overriden by games. */
+	void RegisterEndStreamingPauseRenderingDelegate( FEndStreamingPauseDelegate* InDelegate );
+	FEndStreamingPauseDelegate* EndStreamingPauseDelegate;
+
 	/* 
 	 * Error message event relating to server travel failures 
 	 * 
diff --git a/Engine/Source/Runtime/Engine/Engine.Build.cs b/Engine/Source/Runtime/Engine/Engine.Build.cs
--- a/Engine/Source/Runtime/Engine/Engine.Build.cs
+++ b/Engine/Source/Runtime/Engine/Engine.Build.cs
@@ -9,7 +9,7 @@ public class Engine : ModuleRules
 	{
 		SharedPCHHeaderFile = "Runtime/Engine/Public/Engine.h";
 
-		PublicIncludePathModuleNames.AddRange( new string[] { "Renderer" } );
+		PublicIncludePathModuleNames.AddRange( new string[] { "Renderer", "StreamingPauseRendering" } );
 
 		PrivateIncludePaths.AddRange(
 			new string[] {
@@ -69,6 +69,7 @@ public class Engine : ModuleRules
 				"MovieSceneCore",
 				"MovieSceneCoreTypes",
 				"HeadMountedDisplay",
+				"StreamingPauseRendering",
 			}
 		);
 
diff --git a/Engine/Source/Runtime/Engine/Private/GameEngine.cpp b/Engine/Source/Runtime/Engine/Private/GameEngine.cpp
--- a/Engine/Source/Runtime/Engine/Private/GameEngine.cpp
+++ b/Engine/Source/Runtime/Engine/Private/GameEngine.cpp
@@ -393,6 +393,9 @@ UEngine::UEngine(const class FPostConstructInitializeProperties& PCIP)
 	bUseSound = true;
 	bPendingHardwareSurveyResults = false;
 	bIsInitialized = false;
+
+	BeginStreamingPauseDelegate = NULL;
+	EndStreamingPauseDelegate = NULL;
 }
 
 void UGameEngine::Init(IEngineLoop* InEngineLoop)
@@ -970,12 +973,18 @@ void UGameEngine::Tick( float DeltaSeconds, bool bIdleMode )
 					}
 				}
 
-				FStreamingPause::GameThreadWantsToSuspendRendering( GameViewport ? GameViewport->Viewport : NULL );
+				if( GameViewport && BeginStreamingPauseDelegate && BeginStreamingPauseDelegate->IsBound() )
+				{
+					BeginStreamingPauseDelegate->Execute( GameViewport->Viewport );
+				}	
 
 				// Flushes level streaming requests, blocking till completion.
 				Context.World()->FlushLevelStreaming();
 
-				FStreamingPause::GameThreadWantsToResumeRendering();
+				if( EndStreamingPauseDelegate && EndStreamingPauseDelegate->IsBound() )
+				{
+					EndStreamingPauseDelegate->Execute( );
+				}	
 			}
 			Context.World()->bRequestedBlockOnAsyncLoading = false;
 		}
diff --git a/Engine/Source/Runtime/Engine/Private/StreamingPauseRendering.cpp b/Engine/Source/Runtime/Engine/Private/StreamingPauseRendering.cpp
deleted file mode 100644
diff --git a/Engine/Source/Runtime/Engine/Private/UnrealEngine.cpp b/Engine/Source/Runtime/Engine/Private/UnrealEngine.cpp
--- a/Engine/Source/Runtime/Engine/Private/UnrealEngine.cpp
+++ b/Engine/Source/Runtime/Engine/Private/UnrealEngine.cpp
@@ -844,6 +844,19 @@ void UEngine::Init(IEngineLoop* InEngineLoop)
 		//UGameViewportClient::OnScreenshotCaptured().AddUObject(this, &UEngine::HandleScreenshotCaptured);
 	}
 #endif
+
+	//Load the streaming pause rendering module.
+	FModuleManager::LoadModulePtr<class FStreamingPauseRenderingModule>(TEXT("StreamingPauseRendering"));
+}
+
+void UEngine::RegisterBeginStreamingPauseRenderingDelegate( FBeginStreamingPauseDelegate* InDelegate )
+{
+	BeginStreamingPauseDelegate = InDelegate;
+}
+
+void UEngine::RegisterEndStreamingPauseRenderingDelegate( FEndStreamingPauseDelegate* InDelegate )
+{
+	EndStreamingPauseDelegate = InDelegate;
 }
 
 void UEngine::OnExternalUIChange(bool bInIsOpening)
diff --git a/Engine/Source/Runtime/Engine/Public/StreamingPauseRendering.h b/Engine/Source/Runtime/Engine/Public/StreamingPauseRendering.h
deleted file mode 100644
diff --git a/Engine/Source/Runtime/RenderCore/Private/RenderingThread.cpp b/Engine/Source/Runtime/RenderCore/Private/RenderingThread.cpp
--- a/Engine/Source/Runtime/RenderCore/Private/RenderingThread.cpp
+++ b/Engine/Source/Runtime/RenderCore/Private/RenderingThread.cpp
@@ -35,8 +35,6 @@ FString GRenderingThreadError;
  */
 volatile bool GIsRenderingThreadHealthy = true;
 
-RENDERCORE_API bool GGameThreadWantsToSuspendRendering = false;
-
 /** Whether the rendering thread is suspended (not even processing the tickables) */
 volatile int32 GIsRenderingThreadSuspended = 0;
 
@@ -192,9 +190,7 @@ void TickRenderingTickables()
 	{
 		return;
 	}
-
-	uint32 ObjectsThatResumedRendering = 0;
-
+	
 	// tick any rendering thread tickables
 	for (int32 ObjectIndex = 0; ObjectIndex < FTickableObjectRenderThread::RenderingThreadTickableObjects.Num(); ObjectIndex++)
 	{
@@ -202,22 +198,12 @@ void TickRenderingTickables()
 		// make sure it wants to be ticked and the rendering thread isn't suspended
 		if (TickableObject->IsTickable())
 		{
-			if (GGameThreadWantsToSuspendRendering && TickableObject->NeedsRenderingResumedForRenderingThreadTick())
-			{
-				ObjectsThatResumedRendering++;
-			}
 			STAT(FScopeCycleCounter(TickableObject->GetStatId());)
 			TickableObject->Tick(DeltaSeconds);
 		}
 	}
 	// update the last time we ticked
 	LastTickTime = CurTime;
-
-	// if no ticked objects resumed rendering, make sure we're suspended if game thread wants us to be
-	if (ObjectsThatResumedRendering == 0 && GGameThreadWantsToSuspendRendering)
-	{
-		// Nothing to do?
-	}
 }
 
 /** Accumulates how many cycles the renderthread has been idle. It's defined in RenderingThread.cpp. */
diff --git a/Engine/Source/Runtime/StreamingPauseRendering/Private/StreamingPauseRendering.cpp b/Engine/Source/Runtime/StreamingPauseRendering/Private/StreamingPauseRendering.cpp
new file mode 100644
--- /dev/null
+++ b/Engine/Source/Runtime/StreamingPauseRendering/Private/StreamingPauseRendering.cpp
@@ -0,0 +1,111 @@
+// Copyright 1998-2014 Epic Games, Inc. All Rights Reserved.
+
+/*=============================================================================
+StreamingPauseRendering.cpp: Streaming pause implementation.
+=============================================================================*/
+
+#include "EnginePrivate.h"
+#include "StreamingPauseRendering.h"
+#include "SceneViewport.h"
+#include "MoviePlayer.h"
+
+IMPLEMENT_MODULE(FStreamingPauseRenderingModule, StreamingPauseRenderingModule);
+
+void FStreamingPauseRenderingModule::StartupModule()
+{
+	BeginDelegate.BindRaw(this, &FStreamingPauseRenderingModule::BeginStreamingPause);
+	EndDelegate.BindRaw(this, &FStreamingPauseRenderingModule::EndStreamingPause);
+	check(GEngine);
+	GEngine->RegisterBeginStreamingPauseRenderingDelegate(&BeginDelegate);
+	GEngine->RegisterEndStreamingPauseRenderingDelegate(&EndDelegate);
+}
+
+void FStreamingPauseRenderingModule::ShutdownModule()
+{
+	BeginDelegate.Unbind();
+	EndDelegate.Unbind();
+	if( GEngine )
+	{
+		GEngine->RegisterBeginStreamingPauseRenderingDelegate(NULL);
+		GEngine->RegisterEndStreamingPauseRenderingDelegate(NULL);
+	}
+}
+
+/** Enqueue the streaming pause to suspend rendering during blocking load. */
+void FStreamingPauseRenderingModule::BeginStreamingPause( FViewport* GameViewport )
+{
+#if WITH_STREAMING_PAUSE
+
+	//Create the viewport widget and add a throbber.
+	ViewportWidget = SNew( SViewport )
+		.EnableGammaCorrection(false);
+
+	ViewportWidget->SetContent
+		(
+		SNew(SVerticalBox)
+		+SVerticalBox::Slot()
+		.VAlign(VAlign_Bottom)
+		.HAlign(HAlign_Right)
+		.Padding(FMargin(10.0f))
+		[
+			SNew(SThrobber)
+		]
+	);
+
+	//Render the current scene to a new viewport.
+	static const auto* CVarRenderLastFrameInStreamingPause = IConsoleManager::Get().FindTConsoleVariableDataInt(TEXT("r.RenderLastFrameInStreamingPause"));
+	if( CVarRenderLastFrameInStreamingPause && CVarRenderLastFrameInStreamingPause->GetValueOnGameThread() != 0 ) 
+	{
+		ViewportClient = GameViewport->GetClient();
+
+		Viewport = MakeShareable(new FSceneViewport(ViewportClient, ViewportWidget));
+
+		ViewportWidget->SetViewportInterface(Viewport.ToSharedRef());
+
+		FIntPoint SizeXY = GameViewport->GetSizeXY();	
+		Viewport->UpdateViewportRHI(false,SizeXY.X,SizeXY.Y, EWindowMode::Fullscreen);
+
+		Viewport->EnqueueBeginRenderFrame();
+
+		FCanvas Canvas( Viewport.Get(), NULL, ViewportClient->GetWorld());
+		{
+			ViewportClient->Draw(Viewport.Get(), &Canvas);
+		}
+		Canvas.Flush();
+
+		//Don't need debug canvas I presume?
+		//Viewport->GetDebugCanvas()->Flush(true);
+
+		ENQUEUE_UNIQUE_RENDER_COMMAND_ONEPARAMETER(
+			EndDrawingCommand,
+			FViewport*,Viewport,Viewport.Get(),
+			//FIntPoint,InRestoreSize,RestoreSize,
+		{
+			Viewport->EndRenderFrame( false, false );
+		});
+	}
+
+	//Create the loading screen and begin rendering the widget.
+	FLoadingScreenAttributes LoadingScreen;
+	LoadingScreen.bAutoCompleteWhenLoadingCompletes = true; 
+	LoadingScreen.WidgetLoadingScreen = ViewportWidget; // SViewport from above
+	GetMoviePlayer()->SetupLoadingScreen(LoadingScreen);			
+	GetMoviePlayer()->PlayMovie();
+
+#endif
+
+	FlushRenderingCommands();
+}
+
+/** Enqueue the streaming pause to resume rendering after blocking load is completed. */
+void FStreamingPauseRenderingModule::EndStreamingPause()
+{
+	//Stop rendering the loading screen and resume
+	GetMoviePlayer()->WaitForMovieToFinish();
+
+	ViewportWidget = NULL;
+	Viewport = NULL;
+	ViewportClient = NULL;
+	
+	FlushRenderingCommands();
+}
\ No newline at end of file
diff --git a/Engine/Source/Runtime/StreamingPauseRendering/Public/StreamingPauseRendering.h b/Engine/Source/Runtime/StreamingPauseRendering/Public/StreamingPauseRendering.h
new file mode 100644
--- /dev/null
+++ b/Engine/Source/Runtime/StreamingPauseRendering/Public/StreamingPauseRendering.h
@@ -0,0 +1,46 @@
+// Copyright 1998-2014 Epic Games, Inc. All Rights Reserved.
+
+/*=============================================================================
+	StreamingPauseRendering.h: Streaming pause definitions.
+=============================================================================*/
+
+#ifndef __STREAMINGPAUSERENDERING_H__
+#define __STREAMINGPAUSERENDERING_H__
+
+#include "Engine.h"
+#include "ModuleInterface.h"
+
+/** Allow the streaming pause only on selected platform. It is required for TRC/TCR. */
+#define WITH_STREAMING_PAUSE 1
+
+/** 
+ * Module handling default behavior for streaming pause rendering. 
+ * Games can override by calling RegisterBegin/EndStreamingPauseDelegate with their own delegates.
+ */
+class FStreamingPauseRenderingModule : public IModuleInterface
+{
+public:	
+	FStreamingPauseRenderingModule()
+		: Viewport(NULL)
+		, ViewportWidget(NULL)
+		, ViewportClient(NULL) 
+	{
+	}
+
+	virtual void StartupModule();	
+	virtual void ShutdownModule();
+
+	virtual void BeginStreamingPause( class FViewport* Viewport );
+	virtual void EndStreamingPause();
+
+	TSharedPtr<FSceneViewport>		Viewport;
+	TSharedPtr<SViewport>			ViewportWidget;
+	FViewportClient*			ViewportClient;
+
+	/** Delegate providing default functionality for beginning streaming pause. */
+	FBeginStreamingPauseDelegate BeginDelegate;
+	/** Delegate providing default functionality for ending streaming pause. */
+	FEndStreamingPauseDelegate EndDelegate;
+};
+
+#endif // __STREAMINGPAUSERENDERING_H__
diff --git a/Engine/Source/Runtime/StreamingPauseRendering/StreamingPauseRendering.Build.cs b/Engine/Source/Runtime/StreamingPauseRendering/StreamingPauseRendering.Build.cs
new file mode 100644
--- /dev/null
+++ b/Engine/Source/Runtime/StreamingPauseRendering/StreamingPauseRendering.Build.cs
@@ -0,0 +1,29 @@
+// Copyright 1998-2014 Epic Games, Inc. All Rights Reserved.
+
+using UnrealBuildTool;
+
+public class StreamingPauseRendering : ModuleRules
+{
+    public StreamingPauseRendering(TargetInfo Target)
+	{
+        PrivateIncludePaths.Add("Runtime/StreamingPauseRendering/Private");
+
+		PublicDependencyModuleNames.AddRange(
+			new string[] {
+					"Engine",
+                    "MoviePlayer",
+                    "Slate",
+				}
+		);
+
+		PrivateDependencyModuleNames.AddRange(
+            new string[] {
+                    "Core",
+                    "InputCore",
+                    "RenderCore",
+                    "CoreUObject",
+                    "RHI",
+				}
+        );
+	}
+}
