Fixed compilaton issue with the streaming pause rendering module.

diff --git a/Engine/Source/Runtime/Engine/Engine.Build.cs b/Engine/Source/Runtime/Engine/Engine.Build.cs
--- a/Engine/Source/Runtime/Engine/Engine.Build.cs
+++ b/Engine/Source/Runtime/Engine/Engine.Build.cs
@@ -9,7 +9,7 @@ public class Engine : ModuleRules
 	{
 		SharedPCHHeaderFile = "Runtime/Engine/Public/Engine.h";
 
-		PublicIncludePathModuleNames.AddRange( new string[] { "Renderer", "StreamingPauseRendering" } );
+		PublicIncludePathModuleNames.AddRange( new string[] { "Renderer" } );
 
 		PrivateIncludePaths.AddRange(
 			new string[] {
diff --git a/Engine/Source/Runtime/Engine/Private/GameEngine.cpp b/Engine/Source/Runtime/Engine/Private/GameEngine.cpp
--- a/Engine/Source/Runtime/Engine/Private/GameEngine.cpp
+++ b/Engine/Source/Runtime/Engine/Private/GameEngine.cpp
@@ -15,7 +15,6 @@
 #include "ConfigCacheIni.h"
 
 #include "AVIWriter.h"
-#include "StreamingPauseRendering.h"
 
 #include "Slate.h"
 #include "Slate/SceneViewport.h"
diff --git a/Engine/Source/Runtime/Engine/Private/UnrealEngine.cpp b/Engine/Source/Runtime/Engine/Private/UnrealEngine.cpp
--- a/Engine/Source/Runtime/Engine/Private/UnrealEngine.cpp
+++ b/Engine/Source/Runtime/Engine/Private/UnrealEngine.cpp
@@ -846,7 +846,7 @@ void UEngine::Init(IEngineLoop* InEngineLoop)
 #endif
 
 	//Load the streaming pause rendering module.
-	FModuleManager::LoadModulePtr<class FStreamingPauseRenderingModule>(TEXT("StreamingPauseRendering"));
+	FModuleManager::LoadModulePtr<IModuleInterface>(TEXT("StreamingPauseRendering"));
 }
 
 void UEngine::RegisterBeginStreamingPauseRenderingDelegate( FBeginStreamingPauseDelegate* InDelegate )
diff --git a/Engine/Source/Runtime/StreamingPauseRendering/Private/StreamingPauseRendering.cpp b/Engine/Source/Runtime/StreamingPauseRendering/Private/StreamingPauseRendering.cpp
--- a/Engine/Source/Runtime/StreamingPauseRendering/Private/StreamingPauseRendering.cpp
+++ b/Engine/Source/Runtime/StreamingPauseRendering/Private/StreamingPauseRendering.cpp
@@ -9,7 +9,14 @@ StreamingPauseRendering.cpp: Streaming pause implementation.
 #include "SceneViewport.h"
 #include "MoviePlayer.h"
 
-IMPLEMENT_MODULE(FStreamingPauseRenderingModule, StreamingPauseRenderingModule);
+IMPLEMENT_MODULE(FStreamingPauseRenderingModule, StreamingPauseRendering);
+
+FStreamingPauseRenderingModule::FStreamingPauseRenderingModule()
+: Viewport(NULL)
+, ViewportWidget(NULL)
+, ViewportClient(NULL)
+{
+}
 
 void FStreamingPauseRenderingModule::StartupModule()
 {
@@ -34,8 +41,6 @@ void FStreamingPauseRenderingModule::ShutdownModule()
 /** Enqueue the streaming pause to suspend rendering during blocking load. */
 void FStreamingPauseRenderingModule::BeginStreamingPause( FViewport* GameViewport )
 {
-#if WITH_STREAMING_PAUSE
-
 	//Create the viewport widget and add a throbber.
 	ViewportWidget = SNew( SViewport )
 		.EnableGammaCorrection(false);
@@ -91,10 +96,6 @@ void FStreamingPauseRenderingModule::BeginStreamingPause( FViewport* GameViewpor
 	LoadingScreen.WidgetLoadingScreen = ViewportWidget; // SViewport from above
 	GetMoviePlayer()->SetupLoadingScreen(LoadingScreen);			
 	GetMoviePlayer()->PlayMovie();
-
-#endif
-
-	FlushRenderingCommands();
 }
 
 /** Enqueue the streaming pause to resume rendering after blocking load is completed. */
diff --git a/Engine/Source/Runtime/StreamingPauseRendering/Public/StreamingPauseRendering.h b/Engine/Source/Runtime/StreamingPauseRendering/Public/StreamingPauseRendering.h
--- a/Engine/Source/Runtime/StreamingPauseRendering/Public/StreamingPauseRendering.h
+++ b/Engine/Source/Runtime/StreamingPauseRendering/Public/StreamingPauseRendering.h
@@ -1,5 +1,7 @@
 // Copyright 1998-2014 Epic Games, Inc. All Rights Reserved.
 
+#pragma once
+
 /*=============================================================================
 	StreamingPauseRendering.h: Streaming pause definitions.
 =============================================================================*/
@@ -10,9 +12,6 @@
 #include "Engine.h"
 #include "ModuleInterface.h"
 
-/** Allow the streaming pause only on selected platform. It is required for TRC/TCR. */
-#define WITH_STREAMING_PAUSE 1
-
 /** 
  * Module handling default behavior for streaming pause rendering. 
  * Games can override by calling RegisterBegin/EndStreamingPauseDelegate with their own delegates.
@@ -20,12 +19,7 @@
 class FStreamingPauseRenderingModule : public IModuleInterface
 {
 public:	
-	FStreamingPauseRenderingModule()
-		: Viewport(NULL)
-		, ViewportWidget(NULL)
-		, ViewportClient(NULL) 
-	{
-	}
+	FStreamingPauseRenderingModule();
 
 	virtual void StartupModule();	
 	virtual void ShutdownModule();
