UE-6730 - Crash in Streaming Pause Rendering.

Added check for zero size viewport which seems to happen in recent calls to BeginStreamingPauseRendering. Presumably when blocking for streaming assets when the game window hasn't been created/popped up.

diff --git a/Engine/Source/Runtime/StreamingPauseRendering/Private/StreamingPauseRendering.cpp b/Engine/Source/Runtime/StreamingPauseRendering/Private/StreamingPauseRendering.cpp
--- a/Engine/Source/Runtime/StreamingPauseRendering/Private/StreamingPauseRendering.cpp
+++ b/Engine/Source/Runtime/StreamingPauseRendering/Private/StreamingPauseRendering.cpp
@@ -48,6 +48,8 @@ void FStreamingPauseRenderingModule::ShutdownModule()
 
 void FStreamingPauseRenderingModule::BeginStreamingPause( FViewport* GameViewport )
 {
+	check(GameViewport);
+
 	//Create the viewport widget and add a throbber.
 	ViewportWidget = SNew( SViewport )
 		.EnableGammaCorrection(false);
@@ -65,15 +67,15 @@ void FStreamingPauseRenderingModule::BeginStreamingPause( FViewport* GameViewpor
 	);
 
 	//Render the current scene to a new viewport.
-	if (CVarRenderLastFrameInStreamingPause.GetValueOnGameThread() != 0)
+	FIntPoint SizeXY = GameViewport->GetSizeXY();
+	if (SizeXY.X > 0 && SizeXY.Y > 0 && CVarRenderLastFrameInStreamingPause.GetValueOnGameThread() != 0)
 	{
 		ViewportClient = GameViewport->GetClient();
 
 		Viewport = MakeShareable(new FSceneViewport(ViewportClient, ViewportWidget));
 
 		ViewportWidget->SetViewportInterface(Viewport.ToSharedRef());
-
-		FIntPoint SizeXY = GameViewport->GetSizeXY();	
+	
 		Viewport->UpdateViewportRHI(false,SizeXY.X,SizeXY.Y, EWindowMode::Fullscreen);
 
 		Viewport->EnqueueBeginRenderFrame();
