Do not setup viewport for movie streaming if it is disabled
  - Prevent an unneeded render thread suspension from happening
  - Saves 1s when starting PIE and removes some hitches each time BlockTillLevelStreamingCompleted is called

#rnx

diff --git a/Engine/Source/Runtime/StreamingPauseRendering/Private/StreamingPauseRendering.cpp b/Engine/Source/Runtime/StreamingPauseRendering/Private/StreamingPauseRendering.cpp
--- a/Engine/Source/Runtime/StreamingPauseRendering/Private/StreamingPauseRendering.cpp
+++ b/Engine/Source/Runtime/StreamingPauseRendering/Private/StreamingPauseRendering.cpp
@@ -95,8 +95,10 @@ void FStreamingPauseRenderingModule::ShutdownModule()
 void FStreamingPauseRenderingModule::BeginStreamingPause( FViewport* GameViewport )
 {
 	// If a movie is already playing don't bother starting one
-	if(!GetMoviePlayer()->IsMovieCurrentlyPlaying())
+	if(!GetMoviePlayer()->IsMovieCurrentlyPlaying() && IsMoviePlayerEnabled())
 	{
+		TRACE_CPUPROFILER_EVENT_SCOPE(FStreamingPauseRenderingModule::BeginStreamingPause);
+
 		check(GameViewport);
 
 		//Create the viewport widget and add a throbber.
@@ -161,12 +163,13 @@ void FStreamingPauseRenderingModule::BeginStreamingPause( FViewport* GameViewpor
 	}
 }
 
-
 void FStreamingPauseRenderingModule::EndStreamingPause()
 {
 	// Only wait for the movie to finish if we were the one to start it.
 	if(bMovieWasStarted)
 	{
+		TRACE_CPUPROFILER_EVENT_SCOPE(FStreamingPauseRenderingModule::EndStreamingPause);
+
 		//Stop rendering the loading screen and resume
 		GetMoviePlayer()->WaitForMovieToFinish();
 
