Make sure the MoviePlayer is initialized before playing a movie within the StreamingPauseRendering Module.
The movie player is initialized after engine boot when using the PreLoadScreen module

diff --git a/Engine/Source/Runtime/MoviePlayer/Private/DefaultGameMoviePlayer.h b/Engine/Source/Runtime/MoviePlayer/Private/DefaultGameMoviePlayer.h
--- a/Engine/Source/Runtime/MoviePlayer/Private/DefaultGameMoviePlayer.h
+++ b/Engine/Source/Runtime/MoviePlayer/Private/DefaultGameMoviePlayer.h
@@ -91,6 +91,8 @@ public:
 	virtual void SetupLoadingScreenFromIni() override;
 	virtual void SetViewportDPIScale(float InViewportDPIScale) override;
 	
+	bool IsInitialized() const override { return bInitialized; }
+
 	/** Check if the initial movie(s) is still playing */
 	virtual bool IsStartupMoviePlaying() const override { return IsMoviePlaying; };
 
diff --git a/Engine/Source/Runtime/MoviePlayer/Private/NullMoviePlayer.h b/Engine/Source/Runtime/MoviePlayer/Private/NullMoviePlayer.h
--- a/Engine/Source/Runtime/MoviePlayer/Private/NullMoviePlayer.h
+++ b/Engine/Source/Runtime/MoviePlayer/Private/NullMoviePlayer.h
@@ -37,6 +37,7 @@ public:
 	virtual void PassLoadingScreenWindowBackToGame() const override {}
 	virtual void SetupLoadingScreen(const FLoadingScreenAttributes& InLoadingScreenAttributes) override {}
 	virtual bool HasEarlyStartupMovie() const override { return false; }
+	virtual bool IsInitialized() const override { return false; }
 	virtual bool PlayEarlyStartupMovies() override { return false; }
 	virtual bool PlayMovie() override { return false; }
 	virtual void StopMovie() override {}
diff --git a/Engine/Source/Runtime/MoviePlayer/Public/MoviePlayer.h b/Engine/Source/Runtime/MoviePlayer/Public/MoviePlayer.h
--- a/Engine/Source/Runtime/MoviePlayer/Public/MoviePlayer.h
+++ b/Engine/Source/Runtime/MoviePlayer/Public/MoviePlayer.h
@@ -196,6 +196,11 @@ public:
 	 */
 	virtual bool PlayEarlyStartupMovies() = 0;
 
+	/**
+	 * @return true if Initialize was called
+	 */
+	virtual bool IsInitialized() const = 0;
+
 	/** 
 	 * Starts playing the movie given the last FLoadingScreenAttributes passed in
 	 * @return true of a movie started playing.
diff --git a/Engine/Source/Runtime/StreamingPauseRendering/Private/StreamingPauseRendering.cpp b/Engine/Source/Runtime/StreamingPauseRendering/Private/StreamingPauseRendering.cpp
--- a/Engine/Source/Runtime/StreamingPauseRendering/Private/StreamingPauseRendering.cpp
+++ b/Engine/Source/Runtime/StreamingPauseRendering/Private/StreamingPauseRendering.cpp
@@ -95,7 +95,7 @@ void FStreamingPauseRenderingModule::ShutdownModule()
 void FStreamingPauseRenderingModule::BeginStreamingPause( FViewport* GameViewport )
 {
 	// If a movie is already playing don't bother starting one
-	if(!GetMoviePlayer()->IsMovieCurrentlyPlaying() && IsMoviePlayerEnabled())
+	if(GetMoviePlayer()->IsInitialized() && !GetMoviePlayer()->IsMovieCurrentlyPlaying() && IsMoviePlayerEnabled())
 	{
 		TRACE_CPUPROFILER_EVENT_SCOPE(FStreamingPauseRenderingModule::BeginStreamingPause);
 
